[AutomationTestToolset] Add `AutomationTestToolset` toolset plugin for running automation tests via LLM tool calls.

**Tools:**
- `DiscoverTests()`: initializes session discovery and automation worker communication (async, ~15s first call)
- `ListTests(NameFilter, TagFilter, Limit)`: returns filtered test paths as *JSON*
- `RunTests(TestNames)`: executes tests asynchronously, returns results *JSON* on completion
- `GetTestStatus()`: lightweight polling endpoint for controller state
- `GetTestResults()`: detailed per-test results with errors/warnings
- `StopTests()`: aborts running tests

**Modules:**
- `AutomationTestToolset`: toolset definition + `UEditorSubsystem` for controller lifecycle
- `AutomationTestToolsetTests`: CQTest-based tests (9 tests covering all tools)

#Synced-CL 52608203

diff --git a/Engine/Plugins/Experimental/Toolsets/AutomationTestToolset/AutomationTestToolset.uplugin b/Engine/Plugins/Experimental/Toolsets/AutomationTestToolset/AutomationTestToolset.uplugin
new file mode 100644
--- /dev/null
+++ b/Engine/Plugins/Experimental/Toolsets/AutomationTestToolset/AutomationTestToolset.uplugin
@@ -0,0 +1,38 @@
+{
+	"FileVersion": 3,
+	"Version": 1,
+	"VersionName": "1.0",
+	"FriendlyName": "AutomationTestToolset",
+	"Description": "Automation test discovery and execution tools.",
+	"NoRedist": true,
+	"EditorOnly": true,
+	"CreatedBy": "",
+	"CreatedByURL": "",
+	"DocsURL": "",
+	"MarketplaceURL": "",
+	"CanContainContent": false,
+	"IsBetaVersion": false,
+	"IsExperimentalVersion": true,
+	"Installed": false,
+	"EnabledByDefault": false,
+	"Modules": [
+		{
+			"Name": "AutomationTestToolset",
+			"Type": "Editor",
+			"LoadingPhase": "PostEngineInit",
+			"TargetAllowList": ["Editor"]
+		},
+		{
+			"Name": "AutomationTestToolsetTests",
+			"Type": "Editor",
+			"LoadingPhase": "Default",
+			"TargetAllowList": ["Editor"]
+		}
+	],
+	"Plugins": [
+		{
+			"Name": "ToolsetRegistry",
+			"Enabled": true
+		}
+	]
+}
diff --git a/Engine/Plugins/Experimental/Toolsets/AutomationTestToolset/Source/AutomationTestToolset/AutomationTestToolset.Build.cs b/Engine/Plugins/Experimental/Toolsets/AutomationTestToolset/Source/AutomationTestToolset/AutomationTestToolset.Build.cs
new file mode 100644
--- /dev/null
+++ b/Engine/Plugins/Experimental/Toolsets/AutomationTestToolset/Source/AutomationTestToolset/AutomationTestToolset.Build.cs
@@ -0,0 +1,30 @@
+// Copyright Epic Games, Inc. All Rights Reserved.
+
+namespace UnrealBuildTool.Rules
+{
+	public class AutomationTestToolset : ModuleRules
+	{
+		public AutomationTestToolset(ReadOnlyTargetRules Target) : base(Target)
+		{
+			PublicDependencyModuleNames.AddRange(
+				new string[] {
+					"ToolsetRegistry",
+				}
+			);
+
+			PrivateDependencyModuleNames.AddRange(
+				new string[] {
+					"AutomationController",
+					"AutomationTest",
+					"Core",
+					"CoreUObject",
+					"EditorSubsystem",
+					"Engine",
+					"Json",
+					"SessionServices",
+					"UnrealEd",
+				}
+			);
+		}
+	}
+}
diff --git a/Engine/Plugins/Experimental/Toolsets/AutomationTestToolset/Source/AutomationTestToolset/Private/AutomationTestToolset.cpp b/Engine/Plugins/Experimental/Toolsets/AutomationTestToolset/Source/AutomationTestToolset/Private/AutomationTestToolset.cpp
new file mode 100644
--- /dev/null
+++ b/Engine/Plugins/Experimental/Toolsets/AutomationTestToolset/Source/AutomationTestToolset/Private/AutomationTestToolset.cpp
@@ -0,0 +1,239 @@
+// Copyright Epic Games, Inc. All Rights Reserved.
+
+#include "AutomationTestToolset.h"
+
+#include "AutomationState.h"
+#include "AutomationTestToolsetSubsystem.h"
+#include "Dom/JsonObject.h"
+#include "Dom/JsonValue.h"
+#include "Editor.h"
+#include "IAutomationControllerManager.h"
+#include "IAutomationReport.h"
+#include "Serialization/JsonSerializer.h"
+#include "Serialization/JsonWriter.h"
+#include "ToolsetRegistry/ToolCallAsyncResultString.h"
+
+DEFINE_LOG_CATEGORY_STATIC(LogAutomationTestToolsetTools, Log, All);
+
+// Returns the subsystem only if both it and its automation controller are valid.
+// Tools that need a working controller use this; DiscoverTests bypasses it since
+// it must work before the controller is fully initialized.
+static UAutomationTestToolsetSubsystem* GetSubsystem()
+{
+	if (GEditor == nullptr)
+	{
+		return nullptr;
+	}
+	UAutomationTestToolsetSubsystem* Subsystem = GEditor->GetEditorSubsystem<UAutomationTestToolsetSubsystem>();
+	if (Subsystem == nullptr || !Subsystem->GetAutomationController().IsValid())
+	{
+		return nullptr;
+	}
+	return Subsystem;
+}
+
+static FString JsonObjectToString(const TSharedRef<FJsonObject>& JsonObject)
+{
+	FString OutputString;
+	TSharedRef<TJsonWriter<TCHAR, TPrettyJsonPrintPolicy<TCHAR>>> Writer = TJsonWriterFactory<TCHAR, TPrettyJsonPrintPolicy<TCHAR>>::Create(&OutputString);
+	FJsonSerializer::Serialize(JsonObject, Writer);
+	return OutputString;
+}
+
+UToolCallAsyncResultString* UAutomationTestToolset::DiscoverTests()
+{
+	UAutomationTestToolsetSubsystem* Subsystem = GEditor ? GEditor->GetEditorSubsystem<UAutomationTestToolsetSubsystem>() : nullptr;
+	if (Subsystem == nullptr)
+	{
+		UToolCallAsyncResultString* ErrorResult = NewObject<UToolCallAsyncResultString>();
+		ErrorResult->SetError(TEXT("AutomationTestToolsetSubsystem is not available"));
+		return ErrorResult;
+	}
+
+	IAutomationControllerManagerPtr Controller = Subsystem->GetAutomationController();
+	if (Controller.IsValid() && Controller->GetFilteredReports().Num() > 0)
+	{
+		UToolCallAsyncResultString* Result = NewObject<UToolCallAsyncResultString>();
+		Result->SetValue(TEXT("{\"status\": \"ready\"}"));
+		return Result;
+	}
+
+	UToolCallAsyncResultString* Result = NewObject<UToolCallAsyncResultString>();
+	Subsystem->SetPendingDiscoveryResult(Result);
+	Subsystem->RequestWorkerDiscovery();
+	return Result;
+}
+
+FString UAutomationTestToolset::ListTests(const FString& NameFilter, const FString& TagFilter, int32 Limit)
+{
+	UAutomationTestToolsetSubsystem* Subsystem = GetSubsystem();
+	if (Subsystem == nullptr)
+	{
+		return TEXT("{\"tests\": [], \"total\": 0, \"returned\": 0}");
+	}
+
+	IAutomationControllerManagerPtr Controller = Subsystem->GetAutomationController();
+	TArray<TSharedPtr<IAutomationReport>>& TopLevelReports = Controller->GetFilteredReports();
+
+	TArray<TSharedPtr<IAutomationReport>> LeafReports;
+	UAutomationTestToolsetSubsystem::CollectLeafReports(TopLevelReports, LeafReports);
+
+	TArray<TSharedPtr<FJsonValue>> TestArray;
+	int32 MatchCount = 0;
+
+	for (const TSharedPtr<IAutomationReport>& Report : LeafReports)
+	{
+		if (!Report.IsValid())
+		{
+			continue;
+		}
+
+		const FString& FullPath = Report->GetFullTestPath();
+
+		if (!NameFilter.IsEmpty() && !FullPath.Contains(NameFilter))
+		{
+			continue;
+		}
+		if (!TagFilter.IsEmpty() && !Report->GetTags().Contains(TagFilter))
+		{
+			continue;
+		}
+
+		MatchCount++;
+
+		if (Limit > 0 && TestArray.Num() >= Limit)
+		{
+			continue;
+		}
+
+		TestArray.Add(MakeShared<FJsonValueString>(FullPath));
+	}
+
+	TSharedRef<FJsonObject> Root = MakeShared<FJsonObject>();
+	Root->SetArrayField(TEXT("tests"), TestArray);
+	Root->SetNumberField(TEXT("total"), MatchCount);
+	Root->SetNumberField(TEXT("returned"), TestArray.Num());
+
+	return JsonObjectToString(Root);
+}
+
+UToolCallAsyncResultString* UAutomationTestToolset::RunTests(const TArray<FString>& TestNames)
+{
+	UAutomationTestToolsetSubsystem* Subsystem = GetSubsystem();
+	if (Subsystem == nullptr)
+	{
+		UToolCallAsyncResultString* ErrorResult = NewObject<UToolCallAsyncResultString>();
+		ErrorResult->SetError(TEXT("AutomationTestToolsetSubsystem is not available"));
+		return ErrorResult;
+	}
+
+	if (TestNames.IsEmpty())
+	{
+		UToolCallAsyncResultString* ErrorResult = NewObject<UToolCallAsyncResultString>();
+		ErrorResult->SetError(TEXT("No test names provided"));
+		return ErrorResult;
+	}
+
+	IAutomationControllerManagerPtr Controller = Subsystem->GetAutomationController();
+
+	Controller->StopTests();
+	Controller->SetEnabledTests(TestNames);
+
+	UToolCallAsyncResultString* Result = NewObject<UToolCallAsyncResultString>();
+	Subsystem->SetPendingRunResult(Result, TestNames);
+
+	Controller->RunTests(true);
+
+	// Enable completion polling only after RunTests has set the controller to Running.
+	Subsystem->EnableRunResultPolling();
+
+	return Result;
+}
+
+FString UAutomationTestToolset::GetTestResults()
+{
+	UAutomationTestToolsetSubsystem* Subsystem = GetSubsystem();
+	if (Subsystem == nullptr)
+	{
+		return TEXT("{\"tests\": [], \"passed\": 0, \"failed\": 0, \"skipped\": 0, \"total\": 0, \"duration\": 0.0}");
+	}
+
+	return UAutomationTestToolsetSubsystem::FormatResultsJson(Subsystem->GetAutomationController(), Subsystem->GetRunningTestNames());
+}
+
+FString UAutomationTestToolset::GetTestStatus()
+{
+	UAutomationTestToolsetSubsystem* Subsystem = GetSubsystem();
+	if (Subsystem == nullptr)
+	{
+		return TEXT("{\"state\": \"Disabled\", \"numEnabled\": 0, \"numComplete\": 0, \"numPassed\": 0, \"numFailed\": 0}");
+	}
+
+	IAutomationControllerManagerPtr Controller = Subsystem->GetAutomationController();
+
+	FString StateString = TEXT("Disabled");
+	switch (Controller->GetTestState())
+	{
+	case EAutomationControllerModuleState::Ready:    StateString = TEXT("Ready");    break;
+	case EAutomationControllerModuleState::Running:  StateString = TEXT("Running");  break;
+	case EAutomationControllerModuleState::Disabled: StateString = TEXT("Disabled"); break;
+	}
+
+	TArray<TSharedPtr<IAutomationReport>> EnabledReports = Controller->GetEnabledReports();
+	TArray<TSharedPtr<IAutomationReport>> LeafReports;
+	UAutomationTestToolsetSubsystem::CollectLeafReports(EnabledReports, LeafReports);
+
+	int32 NumComplete = 0;
+	int32 NumPassed = 0;
+	int32 NumFailed = 0;
+	for (const TSharedPtr<IAutomationReport>& Report : LeafReports)
+	{
+		if (!Report.IsValid())
+		{
+			continue;
+		}
+		EAutomationState State = Report->GetState(0, 0);
+		if (State == EAutomationState::Success)
+		{
+			NumComplete++;
+			NumPassed++;
+		}
+		else if (State == EAutomationState::Fail)
+		{
+			NumComplete++;
+			NumFailed++;
+		}
+		else if (State == EAutomationState::Skipped)
+		{
+			NumComplete++;
+		}
+	}
+
+	TSharedRef<FJsonObject> Root = MakeShared<FJsonObject>();
+	Root->SetStringField(TEXT("state"), StateString);
+	Root->SetNumberField(TEXT("numEnabled"), Controller->GetEnabledTestsNum());
+	Root->SetNumberField(TEXT("numComplete"), NumComplete);
+	Root->SetNumberField(TEXT("numPassed"), NumPassed);
+	Root->SetNumberField(TEXT("numFailed"), NumFailed);
+
+	return JsonObjectToString(Root);
+}
+
+bool UAutomationTestToolset::StopTests()
+{
+	UAutomationTestToolsetSubsystem* Subsystem = GetSubsystem();
+	if (Subsystem == nullptr)
+	{
+		return false;
+	}
+
+	Subsystem->GetAutomationController()->StopTests();
+
+	UToolCallAsyncResultString* PendingResult = Subsystem->GetPendingRunResult();
+	if (PendingResult != nullptr && !PendingResult->bIsComplete)
+	{
+		PendingResult->SetError(TEXT("Tests stopped by user"));
+	}
+
+	return true;
+}
diff --git a/Engine/Plugins/Experimental/Toolsets/AutomationTestToolset/Source/AutomationTestToolset/Private/AutomationTestToolsetModule.cpp b/Engine/Plugins/Experimental/Toolsets/AutomationTestToolset/Source/AutomationTestToolset/Private/AutomationTestToolsetModule.cpp
new file mode 100644
--- /dev/null
+++ b/Engine/Plugins/Experimental/Toolsets/AutomationTestToolset/Source/AutomationTestToolset/Private/AutomationTestToolsetModule.cpp
@@ -0,0 +1,15 @@
+// Copyright Epic Games, Inc. All Rights Reserved.
+
+#include "AutomationTestToolsetModule.h"
+
+#include "Modules/ModuleManager.h"
+
+IMPLEMENT_MODULE(FAutomationTestToolsetModule, AutomationTestToolset)
+
+void FAutomationTestToolsetModule::StartupModule()
+{
+}
+
+void FAutomationTestToolsetModule::ShutdownModule()
+{
+}
diff --git a/Engine/Plugins/Experimental/Toolsets/AutomationTestToolset/Source/AutomationTestToolset/Private/AutomationTestToolsetSubsystem.cpp b/Engine/Plugins/Experimental/Toolsets/AutomationTestToolset/Source/AutomationTestToolset/Private/AutomationTestToolsetSubsystem.cpp
new file mode 100644
--- /dev/null
+++ b/Engine/Plugins/Experimental/Toolsets/AutomationTestToolset/Source/AutomationTestToolset/Private/AutomationTestToolsetSubsystem.cpp
@@ -0,0 +1,330 @@
+// Copyright Epic Games, Inc. All Rights Reserved.
+
+#include "AutomationTestToolsetSubsystem.h"
+
+#include "AutomationTestToolset.h"
+#include "AutomationState.h"
+#include "Containers/Ticker.h"
+#include "Dom/JsonObject.h"
+#include "Dom/JsonValue.h"
+#include "Editor.h"
+#include "IAutomationControllerModule.h"
+#include "ISessionInfo.h"
+#include "ISessionManager.h"
+#include "ISessionServicesModule.h"
+#include "Misc/App.h"
+#include "Misc/FilterCollection.h"
+#include "Serialization/JsonSerializer.h"
+#include "Serialization/JsonWriter.h"
+#include "Subsystems/SubsystemCollection.h"
+#include "ToolsetRegistry/ToolCallAsyncResultString.h"
+#include "ToolsetRegistry/UToolsetRegistry.h"
+
+static bool bEnableAutomationTestToolset = true;
+
+static void OnEnableAutomationTestToolsetChanged(IConsoleVariable* Variable)
+{
+	if (GEditor == nullptr)
+	{
+		return;
+	}
+
+	if (UAutomationTestToolsetSubsystem* Subsystem = GEditor->GetEditorSubsystem<UAutomationTestToolsetSubsystem>())
+	{
+		Subsystem->SetToolsetEnabled(Variable->GetBool());
+	}
+}
+
+static FAutoConsoleVariableRef CVarEnableAutomationTestToolset(
+	TEXT("AutomationTestToolset.Enable"),
+	bEnableAutomationTestToolset,
+	TEXT("Enable or disable AutomationTestToolset registration. When disabled, its tools will not appear in the MCP tool list."),
+	FConsoleVariableDelegate::CreateStatic(&OnEnableAutomationTestToolsetChanged),
+	ECVF_Default);
+
+void UAutomationTestToolsetSubsystem::Initialize(FSubsystemCollectionBase& Collection)
+{
+	Super::Initialize(Collection);
+
+	IAutomationControllerModule& ControllerModule = FModuleManager::LoadModuleChecked<IAutomationControllerModule>("AutomationController");
+	AutomationController = ControllerModule.GetAutomationController();
+
+	TestsAvailableHandle = AutomationController->OnTestsAvailable().AddUObject(this, &UAutomationTestToolsetSubsystem::HandleTestsAvailable);
+	TestsRefreshedHandle = AutomationController->OnTestsRefreshed().AddUObject(this, &UAutomationTestToolsetSubsystem::HandleTestsRefreshed);
+
+	TickerHandle = FTSTicker::GetCoreTicker().AddTicker(FTickerDelegate::CreateUObject(this, &UAutomationTestToolsetSubsystem::OnTick));
+
+	SetToolsetEnabled(bEnableAutomationTestToolset);
+}
+
+void UAutomationTestToolsetSubsystem::Deinitialize()
+{
+	FTSTicker::RemoveTicker(TickerHandle);
+
+	if (AutomationController.IsValid())
+	{
+		AutomationController->OnTestsAvailable().Remove(TestsAvailableHandle);
+		AutomationController->OnTestsRefreshed().Remove(TestsRefreshedHandle);
+	}
+
+	if (PendingDiscoveryResult.IsValid() && !PendingDiscoveryResult->bIsComplete)
+	{
+		PendingDiscoveryResult->SetError(TEXT("Subsystem shutting down"));
+	}
+	PendingDiscoveryResult.Reset();
+
+	if (PendingRunResult.IsValid() && !PendingRunResult->bIsComplete)
+	{
+		PendingRunResult->SetError(TEXT("Subsystem shutting down"));
+	}
+	PendingRunResult.Reset();
+
+	SetToolsetEnabled(false);
+	Super::Deinitialize();
+}
+
+void UAutomationTestToolsetSubsystem::SetToolsetEnabled(bool bEnabled)
+{
+	if (bEnabled && !bToolsetRegistered)
+	{
+		UToolsetRegistry::RegisterToolsetClass(UAutomationTestToolset::StaticClass());
+		bToolsetRegistered = true;
+	}
+	else if (!bEnabled && bToolsetRegistered)
+	{
+		UToolsetRegistry::UnregisterToolsetClass(UAutomationTestToolset::StaticClass());
+		bToolsetRegistered = false;
+	}
+}
+
+void UAutomationTestToolsetSubsystem::RequestWorkerDiscovery()
+{
+	if (!SessionManager.IsValid())
+	{
+		ISessionServicesModule& SessionServicesModule = FModuleManager::LoadModuleChecked<ISessionServicesModule>("SessionServices");
+		SessionManager = SessionServicesModule.GetSessionManager();
+	}
+	DiscoveryElapsedSeconds = 0.0f;
+	bDiscoveryRequested = true;
+	bWorkersRequested = false;
+}
+
+void UAutomationTestToolsetSubsystem::SetPendingDiscoveryResult(UToolCallAsyncResultString* InResult)
+{
+	if (PendingDiscoveryResult.IsValid() && !PendingDiscoveryResult->bIsComplete)
+	{
+		PendingDiscoveryResult->SetError(TEXT("Superseded by new discovery request"));
+	}
+	PendingDiscoveryResult = TStrongObjectPtr<UToolCallAsyncResultString>(InResult);
+}
+
+void UAutomationTestToolsetSubsystem::SetPendingRunResult(UToolCallAsyncResultString* InResult, const TArray<FString>& InTestNames)
+{
+	if (PendingRunResult.IsValid() && !PendingRunResult->bIsComplete)
+	{
+		PendingRunResult->SetError(TEXT("Superseded by new test run"));
+	}
+	PendingRunResult = TStrongObjectPtr<UToolCallAsyncResultString>(InResult);
+	RunningTestNames = TSet<FString>(InTestNames);
+}
+
+UToolCallAsyncResultString* UAutomationTestToolsetSubsystem::GetPendingRunResult() const
+{
+	return PendingRunResult.Get();
+}
+
+void UAutomationTestToolsetSubsystem::EnableRunResultPolling()
+{
+	bWaitingForResults = true;
+}
+
+bool UAutomationTestToolsetSubsystem::OnTick(float DeltaTime)
+{
+	if (!AutomationController.IsValid())
+	{
+		return true;
+	}
+
+	IAutomationControllerModule::Get().Tick();
+	AutomationController->Tick();
+
+	// Discovery: find local session, request workers, wait for test list.
+	// The timeout covers the entire process until HandleTestsRefreshed resolves
+	// the pending discovery result.
+	if (bDiscoveryRequested)
+	{
+		DiscoveryElapsedSeconds += DeltaTime;
+
+		if (DiscoveryElapsedSeconds > DiscoveryTimeoutSeconds)
+		{
+			bDiscoveryRequested = false;
+			if (PendingDiscoveryResult.IsValid() && !PendingDiscoveryResult->bIsComplete)
+			{
+				PendingDiscoveryResult->SetError(TEXT("Discovery timed out"));
+				PendingDiscoveryResult.Reset();
+			}
+		}
+		else if (SessionManager.IsValid())
+		{
+			TSharedPtr<ISessionInfo> SelectedSession = SessionManager->GetSelectedSession();
+			if (!SelectedSession.IsValid())
+			{
+				TArray<TSharedPtr<ISessionInfo>> Sessions;
+				SessionManager->GetSessions(Sessions);
+				for (const TSharedPtr<ISessionInfo>& Session : Sessions)
+				{
+					if (Session->GetSessionId() == FApp::GetSessionId())
+					{
+						SessionManager->SelectSession(Session);
+						SelectedSession = Session;
+						break;
+					}
+				}
+			}
+
+			if (SelectedSession.IsValid() && !bWorkersRequested)
+			{
+				AutomationController->RequestAvailableWorkers(SelectedSession->GetSessionId());
+				bWorkersRequested = true;
+			}
+		}
+	}
+
+	// Poll for test run completion.
+	if (bWaitingForResults && AutomationController->GetTestState() != EAutomationControllerModuleState::Running)
+	{
+		CompleteTestRun();
+	}
+
+	return true;
+}
+
+void UAutomationTestToolsetSubsystem::HandleTestsAvailable(EAutomationControllerModuleState::Type State)
+{
+	if (State == EAutomationControllerModuleState::Ready)
+	{
+		bControllerReady = true;
+	}
+	else if (State == EAutomationControllerModuleState::Disabled)
+	{
+		bControllerReady = false;
+	}
+}
+
+void UAutomationTestToolsetSubsystem::HandleTestsRefreshed()
+{
+	bControllerReady = true;
+	bDiscoveryRequested = false;
+
+	TSharedPtr<AutomationFilterCollection> Filters = MakeShared<AutomationFilterCollection>();
+	AutomationController->SetFilter(Filters);
+	AutomationController->SetVisibleTestsEnabled(true);
+
+	if (PendingDiscoveryResult.IsValid() && !PendingDiscoveryResult->bIsComplete)
+	{
+		PendingDiscoveryResult->SetValue(TEXT("{\"status\": \"ready\"}"));
+		PendingDiscoveryResult.Reset();
+	}
+}
+
+void UAutomationTestToolsetSubsystem::CompleteTestRun()
+{
+	bWaitingForResults = false;
+
+	if (PendingRunResult.IsValid() && !PendingRunResult->bIsComplete)
+	{
+		PendingRunResult->SetValue(FormatResultsJson(AutomationController, RunningTestNames));
+	}
+	PendingRunResult.Reset();
+}
+
+void UAutomationTestToolsetSubsystem::CollectLeafReports(const TArray<TSharedPtr<IAutomationReport>>& Reports, TArray<TSharedPtr<IAutomationReport>>& OutLeaves)
+{
+	for (const TSharedPtr<IAutomationReport>& Report : Reports)
+	{
+		if (!Report.IsValid())
+		{
+			continue;
+		}
+		if (!Report->IsParent())
+		{
+			OutLeaves.Add(Report);
+		}
+		else
+		{
+			CollectLeafReports(Report->GetFilteredChildren(), OutLeaves);
+		}
+	}
+}
+
+FString UAutomationTestToolsetSubsystem::FormatResultsJson(const IAutomationControllerManagerPtr& Controller, const TSet<FString>& TestNames)
+{
+	const TArray<TSharedPtr<IAutomationReport>>& AllReports = Controller->GetFilteredReports();
+
+	int32 Passed = 0;
+	int32 Failed = 0;
+	int32 Skipped = 0;
+	float TotalDuration = 0.0f;
+
+	TArray<TSharedPtr<IAutomationReport>> LeafReports;
+	CollectLeafReports(AllReports, LeafReports);
+
+	TArray<TSharedPtr<FJsonValue>> TestArray;
+
+	for (const TSharedPtr<IAutomationReport>& Report : LeafReports)
+	{
+		if (!Report.IsValid() || !TestNames.Contains(Report->GetFullTestPath()))
+		{
+			continue;
+		}
+
+		EAutomationState State = Report->GetState(0, 0);
+		const FAutomationTestResults& Results = Report->GetResults(0, 0);
+
+		switch (State)
+		{
+		case EAutomationState::Success: Passed++;  break;
+		case EAutomationState::Fail:    Failed++;  break;
+		case EAutomationState::Skipped: Skipped++; break;
+		default: break;
+		}
+
+		TotalDuration += Results.Duration;
+
+		TSharedRef<FJsonObject> TestObject = MakeShared<FJsonObject>();
+		TestObject->SetStringField(TEXT("name"), Report->GetFullTestPath());
+		TestObject->SetStringField(TEXT("state"), AutomationStateToString(State));
+		TestObject->SetNumberField(TEXT("duration"), Results.Duration);
+
+		TArray<TSharedPtr<FJsonValue>> Errors;
+		TArray<TSharedPtr<FJsonValue>> Warnings;
+		for (const FAutomationExecutionEntry& Entry : Results.GetEntries())
+		{
+			if (Entry.Event.Type == EAutomationEventType::Error)
+			{
+				Errors.Add(MakeShared<FJsonValueString>(Entry.Event.Message));
+			}
+			else if (Entry.Event.Type == EAutomationEventType::Warning)
+			{
+				Warnings.Add(MakeShared<FJsonValueString>(Entry.Event.Message));
+			}
+		}
+		TestObject->SetArrayField(TEXT("errors"), Errors);
+		TestObject->SetArrayField(TEXT("warnings"), Warnings);
+
+		TestArray.Add(MakeShared<FJsonValueObject>(TestObject));
+	}
+
+	TSharedRef<FJsonObject> Root = MakeShared<FJsonObject>();
+	Root->SetArrayField(TEXT("tests"), TestArray);
+	Root->SetNumberField(TEXT("passed"), Passed);
+	Root->SetNumberField(TEXT("failed"), Failed);
+	Root->SetNumberField(TEXT("skipped"), Skipped);
+	Root->SetNumberField(TEXT("total"), TestArray.Num());
+	Root->SetNumberField(TEXT("duration"), TotalDuration);
+
+	FString OutputString;
+	TSharedRef<TJsonWriter<TCHAR, TPrettyJsonPrintPolicy<TCHAR>>> Writer = TJsonWriterFactory<TCHAR, TPrettyJsonPrintPolicy<TCHAR>>::Create(&OutputString);
+	FJsonSerializer::Serialize(Root, Writer);
+	return OutputString;
+}
diff --git a/Engine/Plugins/Experimental/Toolsets/AutomationTestToolset/Source/AutomationTestToolset/Public/AutomationTestToolset.h b/Engine/Plugins/Experimental/Toolsets/AutomationTestToolset/Source/AutomationTestToolset/Public/AutomationTestToolset.h
new file mode 100644
--- /dev/null
+++ b/Engine/Plugins/Experimental/Toolsets/AutomationTestToolset/Source/AutomationTestToolset/Public/AutomationTestToolset.h
@@ -0,0 +1,73 @@
+// Copyright Epic Games, Inc. All Rights Reserved.
+
+#pragma once
+
+#include "ToolsetRegistry/UToolsetRegistry.h"
+
+#include "AutomationTestToolset.generated.h"
+
+class UToolCallAsyncResultString;
+
+/**
+ * Automation test discovery and execution toolset.
+ *
+ * Wraps the IAutomationControllerManager API (the same backend the Session
+ * Frontend uses) to let MCP clients list available tests, run them, and
+ * retrieve results.
+ *
+ * Typical workflow:
+ *   1. DiscoverTests() once per session to initialize workers and load the test list.
+ *   2. ListTests() to find tests by name or tag.
+ *   3. RunTests() with the desired test names.
+ *   4. GetTestStatus() / GetTestResults() to monitor and retrieve results.
+ *   5. StopTests() to abort if needed.
+ */
+UCLASS(BlueprintType, MinimalAPI)
+class UAutomationTestToolset : public UToolsetDefinition
+{
+	GENERATED_BODY()
+
+public:
+
+	/** Initialize automation worker discovery and load the test list.
+	 * Must be called once before ListTests or RunTests. Takes several seconds
+	 * as it discovers the local automation worker and enumerates all registered
+	 * tests. Returns an async result that completes with a JSON status object
+	 * when tests are available, or an error if discovery fails. */
+	UFUNCTION(BlueprintCallable, Category = "AutomationTestToolset")
+	static AUTOMATIONTESTTOOLSET_API UToolCallAsyncResultString* DiscoverTests();
+
+	/** List available automation tests. Requires DiscoverTests() to have completed.
+	 * Returns a JSON object: {"tests": ["path1", ...], "total": N, "returned": N}.
+	 * @param NameFilter  Optional substring filter applied to the test's full path.
+	 * @param TagFilter   Optional substring filter applied to the test's tags.
+	 * @param Limit       Maximum number of tests to return (0 = unlimited, default 200). */
+	UFUNCTION(BlueprintCallable, Category = "AutomationTestToolset")
+	static AUTOMATIONTESTTOOLSET_API FString ListTests(const FString& NameFilter = TEXT(""), const FString& TagFilter = TEXT(""), int32 Limit = 200);
+
+	/** Run a set of automation tests by name. Requires DiscoverTests() to have completed.
+	 * Starts executing the specified tests and returns an async result that
+	 * completes with a JSON summary when all tests finish.
+	 * @param TestNames  Array of full test paths as returned by ListTests. */
+	UFUNCTION(BlueprintCallable, Category = "AutomationTestToolset")
+	static AUTOMATIONTESTTOOLSET_API UToolCallAsyncResultString* RunTests(const TArray<FString>& TestNames);
+
+	/** Get detailed results for the current or most recent test run.
+	 * Requires DiscoverTests() to have completed.
+	 * Returns a JSON object with per-test state, duration, errors, and warnings. */
+	UFUNCTION(BlueprintCallable, Category = "AutomationTestToolset")
+	static AUTOMATIONTESTTOOLSET_API FString GetTestResults();
+
+	/** Get a lightweight status snapshot of the automation controller.
+	 * Requires DiscoverTests() to have completed.
+	 * Returns a JSON object with the controller state, enabled test count,
+	 * and completion/pass/fail counts. */
+	UFUNCTION(BlueprintCallable, Category = "AutomationTestToolset")
+	static AUTOMATIONTESTTOOLSET_API FString GetTestStatus();
+
+	/** Stop all currently running tests. Requires DiscoverTests() to have completed.
+	 * If a RunTests async result is pending, it will be completed with an error.
+	 * @return True if the stop request was issued successfully. */
+	UFUNCTION(BlueprintCallable, Category = "AutomationTestToolset")
+	static AUTOMATIONTESTTOOLSET_API bool StopTests();
+};
diff --git a/Engine/Plugins/Experimental/Toolsets/AutomationTestToolset/Source/AutomationTestToolset/Public/AutomationTestToolsetModule.h b/Engine/Plugins/Experimental/Toolsets/AutomationTestToolset/Source/AutomationTestToolset/Public/AutomationTestToolsetModule.h
new file mode 100644
--- /dev/null
+++ b/Engine/Plugins/Experimental/Toolsets/AutomationTestToolset/Source/AutomationTestToolset/Public/AutomationTestToolsetModule.h
@@ -0,0 +1,14 @@
+// Copyright Epic Games, Inc. All Rights Reserved.
+
+#pragma once
+
+#include "Modules/ModuleInterface.h"
+
+class FAutomationTestToolsetModule : public IModuleInterface
+{
+public:
+
+	virtual void StartupModule() override;
+
+	virtual void ShutdownModule() override;
+};
diff --git a/Engine/Plugins/Experimental/Toolsets/AutomationTestToolset/Source/AutomationTestToolset/Public/AutomationTestToolsetSubsystem.h b/Engine/Plugins/Experimental/Toolsets/AutomationTestToolset/Source/AutomationTestToolset/Public/AutomationTestToolsetSubsystem.h
new file mode 100644
--- /dev/null
+++ b/Engine/Plugins/Experimental/Toolsets/AutomationTestToolset/Source/AutomationTestToolset/Public/AutomationTestToolsetSubsystem.h
@@ -0,0 +1,106 @@
+// Copyright Epic Games, Inc. All Rights Reserved.
+
+#pragma once
+
+#include "EditorSubsystem.h"
+#include "IAutomationControllerManager.h"
+#include "ToolsetRegistry/ToolCallAsyncResultString.h"
+#include "UObject/StrongObjectPtr.h"
+
+#include "AutomationTestToolsetSubsystem.generated.h"
+
+class FSubsystemCollectionBase;
+class ISessionManager;
+
+/**
+ * Manages the automation controller lifecycle for the AutomationTestToolset.
+ *
+ * Handles session discovery, worker communication, and test execution polling.
+ * The toolset's static tool functions delegate to this subsystem for state management.
+ */
+UCLASS(MinimalAPI)
+class UAutomationTestToolsetSubsystem : public UEditorSubsystem
+{
+	GENERATED_BODY()
+
+public:
+
+	virtual void Initialize(FSubsystemCollectionBase& Collection) override;
+
+	virtual void Deinitialize() override;
+
+	/** Register or unregister the toolset with UToolsetRegistry at runtime. */
+	void SetToolsetEnabled(bool bEnabled);
+
+	/** Get the automation controller. */
+	IAutomationControllerManagerPtr GetAutomationController() const { return AutomationController; }
+
+	/** Whether workers have been found and the test list is populated. */
+	bool IsControllerReady() const { return bControllerReady; }
+
+	/** Load SessionServices and begin polling for the local session.
+	 * Once found, requests available automation workers. */
+	void RequestWorkerDiscovery();
+
+	/** Store a pending async result that will resolve when test discovery completes. */
+	void SetPendingDiscoveryResult(UToolCallAsyncResultString* InResult);
+
+	/** Store a pending async result for the current test run.
+	 * Resolves with results JSON when tests finish. */
+	void SetPendingRunResult(UToolCallAsyncResultString* InResult, const TArray<FString>& InTestNames);
+
+	/** Get the pending run result, if any. */
+	UToolCallAsyncResultString* GetPendingRunResult() const;
+
+	/** Start polling for test completion. Call after RunTests() has started execution. */
+	void EnableRunResultPolling();
+
+	/** Get the test names from the current or most recent run. */
+	const TSet<FString>& GetRunningTestNames() const { return RunningTestNames; }
+
+	/** Build a JSON results string for the given test names from the controller's report tree. */
+	static FString FormatResultsJson(const IAutomationControllerManagerPtr& Controller, const TSet<FString>& TestNames);
+
+	/** Recursively collect leaf (runnable) reports from a report tree. */
+	static void CollectLeafReports(const TArray<TSharedPtr<IAutomationReport>>& Reports, TArray<TSharedPtr<IAutomationReport>>& OutLeaves);
+
+private:
+
+	static constexpr float DiscoveryTimeoutSeconds = 60.0f;
+
+	bool OnTick(float DeltaTime);
+
+	void HandleTestsAvailable(EAutomationControllerModuleState::Type State);
+
+	void HandleTestsRefreshed();
+
+	void CompleteTestRun();
+
+	TSharedPtr<ISessionManager> SessionManager;
+
+	IAutomationControllerManagerPtr AutomationController;
+
+	TStrongObjectPtr<UToolCallAsyncResultString> PendingDiscoveryResult;
+
+	TStrongObjectPtr<UToolCallAsyncResultString> PendingRunResult;
+
+	TSet<FString> RunningTestNames;
+
+	FTSTicker::FDelegateHandle TickerHandle;
+
+	FDelegateHandle TestsAvailableHandle;
+
+	FDelegateHandle TestsRefreshedHandle;
+
+	float DiscoveryElapsedSeconds = 0.0f;
+
+	bool bToolsetRegistered = false;
+
+	bool bControllerReady = false;
+
+	bool bDiscoveryRequested = false;
+
+	bool bWorkersRequested = false;
+
+	bool bWaitingForResults = false;
+};
diff --git a/Engine/Plugins/Experimental/Toolsets/AutomationTestToolset/Source/AutomationTestToolsetTests/AutomationTestToolsetTests.Build.cs b/Engine/Plugins/Experimental/Toolsets/AutomationTestToolset/Source/AutomationTestToolsetTests/AutomationTestToolsetTests.Build.cs
new file mode 100644
--- /dev/null
+++ b/Engine/Plugins/Experimental/Toolsets/AutomationTestToolset/Source/AutomationTestToolsetTests/AutomationTestToolsetTests.Build.cs
@@ -0,0 +1,21 @@
+// Copyright Epic Games, Inc. All Rights Reserved.
+
+namespace UnrealBuildTool.Rules
+{
+	public class AutomationTestToolsetTests : ModuleRules
+	{
+		public AutomationTestToolsetTests(ReadOnlyTargetRules Target) : base(Target)
+		{
+			PrivateDependencyModuleNames.AddRange(new string[]
+			{
+				"AutomationController",
+				"AutomationTestToolset",
+				"Core",
+				"CoreUObject",
+				"CQTest",
+				"Engine",
+				"UnrealEd",
+			});
+		}
+	}
+}
diff --git a/Engine/Plugins/Experimental/Toolsets/AutomationTestToolset/Source/AutomationTestToolsetTests/Private/AutomationTestToolsetTests.cpp b/Engine/Plugins/Experimental/Toolsets/AutomationTestToolset/Source/AutomationTestToolsetTests/Private/AutomationTestToolsetTests.cpp
new file mode 100644
--- /dev/null
+++ b/Engine/Plugins/Experimental/Toolsets/AutomationTestToolset/Source/AutomationTestToolsetTests/Private/AutomationTestToolsetTests.cpp
@@ -0,0 +1,104 @@
+// Copyright Epic Games, Inc. All Rights Reserved.
+
+#include "CQTest.h"
+
+#include "AutomationTestToolset.h"
+#include "AutomationTestToolsetSubsystem.h"
+#include "Editor.h"
+#include "ToolsetRegistry/ToolCallAsyncResultString.h"
+
+TEST_CLASS(AutomationTestToolsetTest, "Editor.AutomationTestToolset")
+{
+	// ── DiscoverTests ───────────────────────────────────────────────
+
+	TEST_METHOD(DiscoverTests_ReturnsAsyncResult)
+	{
+		UToolCallAsyncResultString* Result = UAutomationTestToolset::DiscoverTests();
+		ASSERT_THAT(IsNotNull(Result));
+	}
+
+	TEST_METHOD(DiscoverTests_SecondCallReturnsReadyImmediately)
+	{
+		// First call triggers discovery (or returns ready if already done).
+		UToolCallAsyncResultString* First = UAutomationTestToolset::DiscoverTests();
+		ASSERT_THAT(IsNotNull(First));
+
+		// If the first call already completed synchronously, second should be instant.
+		if (First->bIsComplete && First->Error.IsEmpty())
+		{
+			UToolCallAsyncResultString* Second = UAutomationTestToolset::DiscoverTests();
+			ASSERT_THAT(IsNotNull(Second));
+			ASSERT_THAT(IsTrue(Second->bIsComplete));
+			ASSERT_THAT(IsTrue(Second->Error.IsEmpty()));
+		}
+	}
+
+	// ── ListTests ───────────────────────────────────────────────────
+
+	TEST_METHOD(ListTests_ReturnsJsonWithTestsArray)
+	{
+		const FString Result = UAutomationTestToolset::ListTests();
+		ASSERT_THAT(IsTrue(Result.Contains(TEXT("\"tests\""))));
+		ASSERT_THAT(IsTrue(Result.Contains(TEXT("\"total\""))));
+		ASSERT_THAT(IsTrue(Result.Contains(TEXT("\"returned\""))));
+	}
+
+	TEST_METHOD(ListTests_NameFilterReducesResults)
+	{
+		const FString AllTests = UAutomationTestToolset::ListTests(TEXT(""), TEXT(""), 0);
+		const FString Filtered = UAutomationTestToolset::ListTests(TEXT("AutomationTestToolset"), TEXT(""), 0);
+
+		// Filtered should have fewer or equal matches.
+		ASSERT_THAT(IsTrue(Filtered.Len() <= AllTests.Len()));
+	}
+
+	TEST_METHOD(ListTests_LimitCapsResults)
+	{
+		const FString Result = UAutomationTestToolset::ListTests(TEXT(""), TEXT(""), 2);
+		// The "returned" count should be at most 2.
+		ASSERT_THAT(IsTrue(Result.Contains(TEXT("\"returned\": 2")) || Result.Contains(TEXT("\"returned\": 1")) || Result.Contains(TEXT("\"returned\": 0"))));
+	}
+
+	// ── GetTestStatus ───────────────────────────────────────────────
+
+	TEST_METHOD(GetTestStatus_ReturnsJsonWithState)
+	{
+		const FString Result = UAutomationTestToolset::GetTestStatus();
+		ASSERT_THAT(IsTrue(Result.Contains(TEXT("\"state\""))));
+		ASSERT_THAT(IsTrue(Result.Contains(TEXT("\"numEnabled\""))));
+		ASSERT_THAT(IsTrue(Result.Contains(TEXT("\"numComplete\""))));
+		ASSERT_THAT(IsTrue(Result.Contains(TEXT("\"numPassed\""))));
+		ASSERT_THAT(IsTrue(Result.Contains(TEXT("\"numFailed\""))));
+	}
+
+	// ── GetTestResults ──────────────────────────────────────────────
+
+	TEST_METHOD(GetTestResults_ReturnsJsonWithTestsArray)
+	{
+		const FString Result = UAutomationTestToolset::GetTestResults();
+		ASSERT_THAT(IsTrue(Result.Contains(TEXT("\"tests\""))));
+		ASSERT_THAT(IsTrue(Result.Contains(TEXT("\"passed\""))));
+		ASSERT_THAT(IsTrue(Result.Contains(TEXT("\"failed\""))));
+		ASSERT_THAT(IsTrue(Result.Contains(TEXT("\"total\""))));
+	}
+
+	// ── RunTests ────────────────────────────────────────────────────
+
+	TEST_METHOD(RunTests_EmptyArrayReturnsError)
+	{
+		UToolCallAsyncResultString* Result = UAutomationTestToolset::RunTests({});
+		ASSERT_THAT(IsNotNull(Result));
+		ASSERT_THAT(IsTrue(Result->bIsComplete));
+		ASSERT_THAT(IsTrue(!Result->Error.IsEmpty()));
+	}
+
+	// ── Subsystem ───────────────────────────────────────────────────
+
+	TEST_METHOD(Subsystem_IsAvailable)
+	{
+		UAutomationTestToolsetSubsystem* Subsystem = GEditor->GetEditorSubsystem<UAutomationTestToolsetSubsystem>();
+		ASSERT_THAT(IsNotNull(Subsystem));
+		ASSERT_THAT(IsTrue(Subsystem->GetAutomationController().IsValid()));
+	}
+
+};
diff --git a/Engine/Plugins/Experimental/Toolsets/AutomationTestToolset/Source/AutomationTestToolsetTests/Private/AutomationTestToolsetTestsModule.cpp b/Engine/Plugins/Experimental/Toolsets/AutomationTestToolset/Source/AutomationTestToolsetTests/Private/AutomationTestToolsetTestsModule.cpp
new file mode 100644
--- /dev/null
+++ b/Engine/Plugins/Experimental/Toolsets/AutomationTestToolset/Source/AutomationTestToolsetTests/Private/AutomationTestToolsetTestsModule.cpp
@@ -0,0 +1,9 @@
+// Copyright Epic Games, Inc. All Rights Reserved.
+
+#include "Modules/ModuleManager.h"
+
+class FAutomationTestToolsetTestsModule : public IModuleInterface
+{
+};
+
+IMPLEMENT_MODULE(FAutomationTestToolsetTestsModule, AutomationTestToolsetTests)
