[AutomationTestToolset] Move `AutomationTestToolset` tests from `Editor` to `AI.Toolsets` category.

#Synced-CL 52658100

diff --git a/Engine/Plugins/Experimental/Toolsets/AutomationTestToolset/Source/AutomationTestToolsetTests/Private/AutomationTestToolsetTests.cpp b/Engine/Plugins/Experimental/Toolsets/AutomationTestToolset/Source/AutomationTestToolsetTests/Private/AutomationTestToolsetTests.cpp
--- a/Engine/Plugins/Experimental/Toolsets/AutomationTestToolset/Source/AutomationTestToolsetTests/Private/AutomationTestToolsetTests.cpp
+++ b/Engine/Plugins/Experimental/Toolsets/AutomationTestToolset/Source/AutomationTestToolsetTests/Private/AutomationTestToolsetTests.cpp
@@ -7,7 +7,7 @@
 #include "Editor.h"
 #include "ToolsetRegistry/ToolCallAsyncResultString.h"
 
-TEST_CLASS(AutomationTestToolsetTest, "Editor.AutomationTestToolset")
+TEST_CLASS(AutomationTestToolsetTest, "AI.Toolsets.AutomationTestToolset")
 {
 	// ── DiscoverTests ───────────────────────────────────────────────
 
