[AutomationTestToolset] Guard `HandleTestsRefreshed` filter reset behind `bDiscoveryRequested`.

The unconditional `SetFilter` + `SetVisibleTestsEnabled(true)` in `HandleTestsRefreshed` was clobbering *Gauntlet* test selection and *Session Frontend* filters whenever any test list refresh occurred. Now the filter reset only runs when the subsystem itself initiated discovery via `DiscoverTests()`.

#Synced-CL 52658100

diff --git a/Engine/Plugins/Experimental/Toolsets/AutomationTestToolset/Source/AutomationTestToolset/Private/AutomationTestToolsetSubsystem.cpp b/Engine/Plugins/Experimental/Toolsets/AutomationTestToolset/Source/AutomationTestToolset/Private/AutomationTestToolsetSubsystem.cpp
--- a/Engine/Plugins/Experimental/Toolsets/AutomationTestToolset/Source/AutomationTestToolset/Private/AutomationTestToolsetSubsystem.cpp
+++ b/Engine/Plugins/Experimental/Toolsets/AutomationTestToolset/Source/AutomationTestToolset/Private/AutomationTestToolsetSubsystem.cpp
@@ -214,16 +214,20 @@ void UAutomationTestToolsetSubsystem::HandleTestsAvailable(EAutomationController
 void UAutomationTestToolsetSubsystem::HandleTestsRefreshed()
 {
 	bControllerReady = true;
-	bDiscoveryRequested = false;
 
-	TSharedPtr<AutomationFilterCollection> Filters = MakeShared<AutomationFilterCollection>();
-	AutomationController->SetFilter(Filters);
-	AutomationController->SetVisibleTestsEnabled(true);
-
-	if (PendingDiscoveryResult.IsValid() && !PendingDiscoveryResult->bIsComplete)
+	if (bDiscoveryRequested)
 	{
-		PendingDiscoveryResult->SetValue(TEXT("{\"status\": \"ready\"}"));
-		PendingDiscoveryResult.Reset();
+		bDiscoveryRequested = false;
+
+		TSharedPtr<AutomationFilterCollection> Filters = MakeShared<AutomationFilterCollection>();
+		AutomationController->SetFilter(Filters);
+		AutomationController->SetVisibleTestsEnabled(true);
+
+		if (PendingDiscoveryResult.IsValid() && !PendingDiscoveryResult->bIsComplete)
+		{
+			PendingDiscoveryResult->SetValue(TEXT("{\"status\": \"ready\"}"));
+			PendingDiscoveryResult.Reset();
+		}
 	}
 }
 
